Fix errors compiling BlankProgram for non-desktop platforms, and corrupt resource files caused by incorrect UTF-16 BOM at the start of files.

diff --git a/Engine/Source/Programs/BlankProgram/BlankProgram.Target.cs b/Engine/Source/Programs/BlankProgram/BlankProgram.Target.cs
--- a/Engine/Source/Programs/BlankProgram/BlankProgram.Target.cs
+++ b/Engine/Source/Programs/BlankProgram/BlankProgram.Target.cs
@@ -3,7 +3,7 @@
 using UnrealBuildTool;
 using System.Collections.Generic;
 
-[SupportedPlatforms(UnrealPlatformClass.All)]
+[SupportedPlatforms(UnrealPlatformClass.Desktop)]
 public class BlankProgramTarget : TargetRules
 {
 	public BlankProgramTarget(TargetInfo Target) : base(Target)
diff --git a/Engine/Source/Programs/BlankProgram/Resources/Windows/Resource.rc b/Engine/Source/Programs/BlankProgram/Resources/Windows/Resource.rc
--- a/Engine/Source/Programs/BlankProgram/Resources/Windows/Resource.rc
+++ b/Engine/Source/Programs/BlankProgram/Resources/Windows/Resource.rc
@@ -1,4 +1,4 @@
-ÿþ// Copyright 1998-2018 Epic Games, Inc. All Rights Reserved.
+// Copyright 1998-2018 Epic Games, Inc. All Rights Reserved.
 
 #include <windows.h>
 #include "Runtime/Launch/Resources/Version.h"
