BlankProgram support for all targets to ease running unit tests and benchmarks on consoles and mobile platforms

BlankProgram compiles for all platforms, but I've only tested running it on PS5.
Also added generation of PS5 VS debugger settings for all Program targets.

diff --git a/Engine/Source/Programs/BlankProgram/BlankProgram.Target.cs b/Engine/Source/Programs/BlankProgram/BlankProgram.Target.cs
--- a/Engine/Source/Programs/BlankProgram/BlankProgram.Target.cs
+++ b/Engine/Source/Programs/BlankProgram/BlankProgram.Target.cs
@@ -3,7 +3,7 @@
 using UnrealBuildTool;
 using System.Collections.Generic;
 
-[SupportedPlatforms(UnrealPlatformClass.Desktop)]
+[SupportedPlatforms(UnrealPlatformClass.All)]
 public class BlankProgramTarget : TargetRules
 {
 	public BlankProgramTarget(TargetInfo Target) : base(Target)
@@ -19,13 +19,16 @@ public class BlankProgramTarget : TargetRules
 		// automatically by Unreal Build Tool, but if bUseMallocProfiler is defined, UHT can operate incorrectly.
 		bUseMallocProfiler = false;
 
-		// Editor-only data, however, is needed
-		bBuildWithEditorOnlyData = true;
+		// Editor-only is enabled for desktop platforms to run unit tests that depend on editor-only data
+		// It's disabled in test and shipping configs to make profiling similar to the game
+		bool bDebugOrDevelopment = Target.Configuration == UnrealTargetConfiguration.Debug || Target.Configuration == UnrealTargetConfiguration.Development;
+		bBuildWithEditorOnlyData = Target.Platform.IsInGroup(UnrealPlatformGroup.Desktop) && bDebugOrDevelopment;
 
 		// Currently this app is not linking against the engine, so we'll compile out references from Core to the rest of the engine
 		bCompileAgainstEngine = false;
 		bCompileAgainstCoreUObject = false;
 		bCompileAgainstApplicationCore = false;
+		bCompileICU = false;
 
 		// UnrealHeaderTool is a console application, not a Windows app (sets entry point to main(), instead of WinMain())
 		bIsBuildingConsoleApplication = true;
