* BlankProgram can be linked with CoreUObject and exits correctly
* instructions how to run automation tests in BlankProgram, how to enable tracing and other tips

diff --git a/Engine/Source/Programs/BlankProgram/BlankProgram.Build.cs b/Engine/Source/Programs/BlankProgram/BlankProgram.Build.cs
--- a/Engine/Source/Programs/BlankProgram/BlankProgram.Build.cs
+++ b/Engine/Source/Programs/BlankProgram/BlankProgram.Build.cs
@@ -9,5 +9,15 @@ public class BlankProgram : ModuleRules
 		PublicIncludePathModuleNames.Add("Launch");
 		PrivateDependencyModuleNames.Add("Core");
 		PrivateDependencyModuleNames.Add("Projects");
+
+		// to link with CoreUObject module:
+		// PrivateDependencyModuleNames.Add("CoreUObject");
+
+		// to enable tracing:
+		// AppendStringToPublicDefinition("UE_TRACE_ENABLED", "1");
+
+		// to enable LLM tracing:
+		// GlobalDefinitions.Add("LLM_ENABLED_IN_CONFIG=1");
+		// GlobalDefinitions.Add("UE_MEMORY_TAGS_TRACE_ENABLED=1");
 	}
 }
diff --git a/Engine/Source/Programs/BlankProgram/BlankProgram.Target.cs b/Engine/Source/Programs/BlankProgram/BlankProgram.Target.cs
--- a/Engine/Source/Programs/BlankProgram/BlankProgram.Target.cs
+++ b/Engine/Source/Programs/BlankProgram/BlankProgram.Target.cs
@@ -31,6 +31,12 @@ public class BlankProgramTarget : TargetRules
 		bCompileAgainstApplicationCore = false;
 		bCompileICU = false;
 
+		// to build with automation tests:
+		// bForceCompileDevelopmentAutomationTests = true;
+
+		// to enable tracing:
+		// GlobalDefinitions.Add("UE_TRACE_ENABLED=1");
+
 		// This app is a console application, not a Windows app (sets entry point to main(), instead of WinMain())
 		bIsBuildingConsoleApplication = true;
 	}
diff --git a/Engine/Source/Programs/BlankProgram/Private/BlankProgram.cpp b/Engine/Source/Programs/BlankProgram/Private/BlankProgram.cpp
--- a/Engine/Source/Programs/BlankProgram/Private/BlankProgram.cpp
+++ b/Engine/Source/Programs/BlankProgram/Private/BlankProgram.cpp
@@ -1,6 +1,5 @@
 // Copyright Epic Games, Inc. All Rights Reserved.
 
-
 #include "BlankProgram.h"
 
 #include "RequiredProgramMainCPPInclude.h"
@@ -11,8 +10,41 @@ IMPLEMENT_APPLICATION(BlankProgram, "BlankProgram");
 
 INT32_MAIN_INT32_ARGC_TCHAR_ARGV()
 {
-	GEngineLoop.PreInit(ArgC, ArgV);
+	FTaskTagScope Scope(ETaskTag::EGameThread);
+	ON_SCOPE_EXIT
+	{ 
+		LLM(FLowLevelMemTracker::Get().UpdateStatsPerFrame());
+		RequestEngineExit(TEXT("Exiting"));
+		FEngineLoop::AppPreExit();
+		FModuleManager::Get().UnloadModulesAtShutdown();
+		FEngineLoop::AppExit();
+	};
+
+	if (int32 Ret = GEngineLoop.PreInit(ArgC, ArgV))
+	{
+		return Ret;
+	}
+
+	// to run automation tests in BlankProgram:
+	// * set `bForceCompileDevelopmentAutomationTests` to `true` in `BlankProgramTarget` constructor in "BlankProgram.Target.cs"
+	// add `FAutomationTestFramework::Get().StartTestByName(TEXT("The name of the test class passed to IMPLEMENT_SIMPLE_AUTOMATION_TEST or TEST_CASE_NAMED"), 0);` right here
+
+	// to link with "CoreUObject" module:
+	// * uncomment `PrivateDependencyModuleNames.Add("CoreUObject");` in `BlankProgram` constructor in "BlankProgram.Build.cs"
+	// * set `bCompileAgainstCoreUObject` to `true` in `BlankProgramTarget` constructor in "BlankProgram.Target.cs"
+
+	// to enable tracing:
+	// * uncomment `AppendStringToPublicDefinition("UE_TRACE_ENABLED", "1");` in `BlankProgram` constructor in "BlankProgram.Build.cs"
+	// * uncomment `GlobalDefinitions.Add("UE_TRACE_ENABLED=1");` in `BlankProgramTarget` constructor in "BlankProgram.Target.cs"
+	// you may need to enable compilation of a particular trace channel, e.g. for "task" traces:
+	// * add `GlobalDefinitions.Add("UE_TASK_TRACE_ENABLED=1");` in `BlankProgramTarget` constructor in "BlankProgram.Target.cs"
+	// you'll still need to enable this trace channel on cmd-line like `-trace=task,default`
+
+	// to enable LLM tracing, uncomment the following in `BlankProgram` constructor in "BlankProgram.Build.cs":
+	// GlobalDefinitions.Add("LLM_ENABLED_IN_CONFIG=1");
+	// GlobalDefinitions.Add("UE_MEMORY_TAGS_TRACE_ENABLED=1");
+
 	UE_LOG(LogBlankProgram, Display, TEXT("Hello World"));
-	FEngineLoop::AppExit();
+
 	return 0;
 }
