UBT: Removed deprecated TargetInfo::bUseMallocProfiler toggle in Source/Programs.

diff --git a/Engine/Source/Programs/BenchmarkTool/BenchmarkTool.Target.cs b/Engine/Source/Programs/BenchmarkTool/BenchmarkTool.Target.cs
--- a/Engine/Source/Programs/BenchmarkTool/BenchmarkTool.Target.cs
+++ b/Engine/Source/Programs/BenchmarkTool/BenchmarkTool.Target.cs
@@ -14,7 +14,6 @@ public class BenchmarkToolTarget : TargetRules
 		LaunchModuleName = "BenchmarkTool";
 
 		bBuildDeveloperTools			= false;
-		bUseMallocProfiler				= false;
 		bBuildWithEditorOnlyData		= false;
 		bCompileAgainstEngine			= false;
 		bCompileAgainstCoreUObject		= false;
diff --git a/Engine/Source/Programs/BlankProgram/BlankProgram.Target.cs b/Engine/Source/Programs/BlankProgram/BlankProgram.Target.cs
--- a/Engine/Source/Programs/BlankProgram/BlankProgram.Target.cs
+++ b/Engine/Source/Programs/BlankProgram/BlankProgram.Target.cs
@@ -16,10 +16,6 @@ public class BlankProgramTarget : TargetRules
 		// Lean and mean
 		bBuildDeveloperTools = false;
 
-		// Never use malloc profiling in Unreal Header Tool.  We set this because often UHT is compiled right before the engine
-		// automatically by Unreal Build Tool, but if bUseMallocProfiler is defined, UHT can operate incorrectly.
-		bUseMallocProfiler = false;
-
 		// Editor-only is enabled for desktop platforms to run unit tests that depend on editor-only data
 		// It's disabled in test and shipping configs to make profiling similar to the game
 		bool bDebugOrDevelopment = Target.Configuration == UnrealTargetConfiguration.Debug || Target.Configuration == UnrealTargetConfiguration.Development;
diff --git a/Engine/Source/Programs/ChaosVisualDebugger/ChaosVisualDebugger.Target.cs b/Engine/Source/Programs/ChaosVisualDebugger/ChaosVisualDebugger.Target.cs
--- a/Engine/Source/Programs/ChaosVisualDebugger/ChaosVisualDebugger.Target.cs
+++ b/Engine/Source/Programs/ChaosVisualDebugger/ChaosVisualDebugger.Target.cs
@@ -15,10 +15,6 @@ public class ChaosVisualDebuggerTarget : TargetRules
 		// Lean and mean
 		bBuildDeveloperTools = true;
 
-		// Never use malloc profiling in Unreal Header Tool.  We set this because often UHT is compiled right before the engine
-		// automatically by Unreal Build Tool, but if bUseMallocProfiler is defined, UHT can operate incorrectly.
-		bUseMallocProfiler = false;
-
 		// Editor-only data, however, is needed
 		bBuildWithEditorOnlyData = true;
 
diff --git a/Engine/Source/Programs/DerivedDataBuildWorker/DerivedDataBuildWorker.Build.cs b/Engine/Source/Programs/DerivedDataBuildWorker/DerivedDataBuildWorker.Build.cs
--- a/Engine/Source/Programs/DerivedDataBuildWorker/DerivedDataBuildWorker.Build.cs
+++ b/Engine/Source/Programs/DerivedDataBuildWorker/DerivedDataBuildWorker.Build.cs
@@ -25,7 +25,6 @@ public abstract class DerivedDataBuildWorkerTarget : TargetRules
 		bCompileICU 					= false;
 		bWithLiveCoding					= false;
 		bBuildDeveloperTools			= false;
-		bUseMallocProfiler				= false;
 		bBuildWithEditorOnlyData		= true;
 		bCompileAgainstEngine			= false;
 		bCompileAgainstCoreUObject		= false;
diff --git a/Engine/Source/Programs/Enterprise/Datasmith/DatasmithARCHICADExporter/DatsmithUE4ArchiCAD/DatasmithUE4ArchiCAD.Target.cs b/Engine/Source/Programs/Enterprise/Datasmith/DatasmithARCHICADExporter/DatsmithUE4ArchiCAD/DatasmithUE4ArchiCAD.Target.cs
--- a/Engine/Source/Programs/Enterprise/Datasmith/DatasmithARCHICADExporter/DatsmithUE4ArchiCAD/DatasmithUE4ArchiCAD.Target.cs
+++ b/Engine/Source/Programs/Enterprise/Datasmith/DatasmithARCHICADExporter/DatsmithUE4ArchiCAD/DatasmithUE4ArchiCAD.Target.cs
@@ -26,7 +26,6 @@ public class DatasmithUE4ArchiCADTarget : TargetRules
 		}
 
 		bBuildDeveloperTools = false;
-		bUseMallocProfiler = false;
 		bBuildWithEditorOnlyData = true;
 		bCompileAgainstEngine = false;
 		bCompileAgainstCoreUObject = true;
diff --git a/Engine/Source/Programs/Enterprise/Datasmith/DatasmithARCHICADExporter/UBTModules/DatasmithArchiCAD.Target.cs b/Engine/Source/Programs/Enterprise/Datasmith/DatasmithARCHICADExporter/UBTModules/DatasmithArchiCAD.Target.cs
--- a/Engine/Source/Programs/Enterprise/Datasmith/DatasmithARCHICADExporter/UBTModules/DatasmithArchiCAD.Target.cs
+++ b/Engine/Source/Programs/Enterprise/Datasmith/DatasmithARCHICADExporter/UBTModules/DatasmithArchiCAD.Target.cs
@@ -21,7 +21,6 @@ public class DatasmithArchiCADTarget : TargetRules
 		LinkType = TargetLinkType.Monolithic;
 
 		bBuildDeveloperTools = false;
-		bUseMallocProfiler = false;
 		bBuildWithEditorOnlyData = true;
 		bCompileAgainstEngine = false;
 		bCompileAgainstCoreUObject = true;
diff --git a/Engine/Source/Programs/Enterprise/Datasmith/DatasmithCADWorker/DatasmithCADWorker.Target.cs b/Engine/Source/Programs/Enterprise/Datasmith/DatasmithCADWorker/DatasmithCADWorker.Target.cs
--- a/Engine/Source/Programs/Enterprise/Datasmith/DatasmithCADWorker/DatasmithCADWorker.Target.cs
+++ b/Engine/Source/Programs/Enterprise/Datasmith/DatasmithCADWorker/DatasmithCADWorker.Target.cs
@@ -17,10 +17,6 @@ public class DatasmithCADWorkerTarget : TargetRules
         // Lean and mean
         bBuildDeveloperTools = false;
 
-		// Never use malloc profiling in Unreal Header Tool.  We set this because often UHT is compiled right before the engine
-		// automatically by Unreal Build Tool, but if bUseMallocProfiler is defined, UHT can operate incorrectly.
-		bUseMallocProfiler = false;
-
 		// Editor-only data, however, is needed
 		bBuildWithEditorOnlyData = false;
 
diff --git a/Engine/Source/Programs/Enterprise/Datasmith/DatasmithFacadeCSharp/DatasmithFacadeCSharp.Target.cs b/Engine/Source/Programs/Enterprise/Datasmith/DatasmithFacadeCSharp/DatasmithFacadeCSharp.Target.cs
--- a/Engine/Source/Programs/Enterprise/Datasmith/DatasmithFacadeCSharp/DatasmithFacadeCSharp.Target.cs
+++ b/Engine/Source/Programs/Enterprise/Datasmith/DatasmithFacadeCSharp/DatasmithFacadeCSharp.Target.cs
@@ -22,7 +22,6 @@ public class DatasmithFacadeCSharpTarget : TargetRules
 		LinkType = TargetLinkType.Monolithic;
 
 		bBuildDeveloperTools = false;
-		bUseMallocProfiler = false;
 		bBuildWithEditorOnlyData = true;
 		bCompileAgainstEngine = false;
 		bCompileAgainstCoreUObject = true;
diff --git a/Engine/Source/Programs/Enterprise/Datasmith/DatasmithMaxExporter/DatasmithMax2017.Target.cs b/Engine/Source/Programs/Enterprise/Datasmith/DatasmithMaxExporter/DatasmithMax2017.Target.cs
--- a/Engine/Source/Programs/Enterprise/Datasmith/DatasmithMaxExporter/DatasmithMax2017.Target.cs
+++ b/Engine/Source/Programs/Enterprise/Datasmith/DatasmithMaxExporter/DatasmithMax2017.Target.cs
@@ -22,7 +22,6 @@ public abstract class DatasmithMaxBaseTarget : TargetRules
 		WindowsPlatform.bStrictConformanceMode = false;
 
 		bBuildDeveloperTools = false;
-		bUseMallocProfiler = false;
 		bBuildWithEditorOnlyData = true;
 		bCompileAgainstEngine = false;
 		bCompileAgainstCoreUObject = true;
diff --git a/Engine/Source/Programs/Enterprise/Datasmith/DatasmithNavisworksExporter/UBTModules/DatasmithNavisworks2020.Target.cs b/Engine/Source/Programs/Enterprise/Datasmith/DatasmithNavisworksExporter/UBTModules/DatasmithNavisworks2020.Target.cs
--- a/Engine/Source/Programs/Enterprise/Datasmith/DatasmithNavisworksExporter/UBTModules/DatasmithNavisworks2020.Target.cs
+++ b/Engine/Source/Programs/Enterprise/Datasmith/DatasmithNavisworksExporter/UBTModules/DatasmithNavisworks2020.Target.cs
@@ -23,7 +23,6 @@ public abstract class DatasmithNavisworksBaseTarget : TargetRules
 		LinkType = TargetLinkType.Monolithic;
 
 		bBuildDeveloperTools = false;
-		bUseMallocProfiler = false;
 		bBuildWithEditorOnlyData = true;
 		bCompileAgainstEngine = false;
 		bCompileAgainstCoreUObject = false;
diff --git a/Engine/Source/Programs/Enterprise/Datasmith/DatasmithRevitExporter/DatasmithRevit2018.Target.cs b/Engine/Source/Programs/Enterprise/Datasmith/DatasmithRevitExporter/DatasmithRevit2018.Target.cs
--- a/Engine/Source/Programs/Enterprise/Datasmith/DatasmithRevitExporter/DatasmithRevit2018.Target.cs
+++ b/Engine/Source/Programs/Enterprise/Datasmith/DatasmithRevitExporter/DatasmithRevit2018.Target.cs
@@ -25,7 +25,6 @@ public abstract class DatasmithRevitBaseTarget : TargetRules
 		LinkType = TargetLinkType.Monolithic;
 
 		bBuildDeveloperTools = false;
-		bUseMallocProfiler = false;
 		bBuildWithEditorOnlyData = true;
 		bCompileAgainstEngine = false;
 		bCompileAgainstCoreUObject = true;
diff --git a/Engine/Source/Programs/Enterprise/Datasmith/DatasmithRhinoExporter/DatasmithRhino6.Target.cs b/Engine/Source/Programs/Enterprise/Datasmith/DatasmithRhinoExporter/DatasmithRhino6.Target.cs
--- a/Engine/Source/Programs/Enterprise/Datasmith/DatasmithRhinoExporter/DatasmithRhino6.Target.cs
+++ b/Engine/Source/Programs/Enterprise/Datasmith/DatasmithRhinoExporter/DatasmithRhino6.Target.cs
@@ -34,7 +34,6 @@ public abstract class DatasmithRhinoBaseTarget : TargetRules
 		LinkType = TargetLinkType.Monolithic;
 
 		bBuildDeveloperTools = false;
-		bUseMallocProfiler = false;
 		bBuildWithEditorOnlyData = true;
 		bCompileAgainstEngine = false;
 		bCompileAgainstCoreUObject = true;
diff --git a/Engine/Source/Programs/Enterprise/Datasmith/DatasmithSDK/DatasmithSDK.Target.cs b/Engine/Source/Programs/Enterprise/Datasmith/DatasmithSDK/DatasmithSDK.Target.cs
--- a/Engine/Source/Programs/Enterprise/Datasmith/DatasmithSDK/DatasmithSDK.Target.cs
+++ b/Engine/Source/Programs/Enterprise/Datasmith/DatasmithSDK/DatasmithSDK.Target.cs
@@ -30,7 +30,6 @@ public class DatasmithSDKTarget : TargetRules
 		bShouldCompileAsDLL = true;
 
 		bBuildDeveloperTools = false;
-		bUseMallocProfiler = false;
 		bBuildWithEditorOnlyData = true;
 		bCompileAgainstEngine = false;
 		bCompileAgainstCoreUObject = true;
diff --git a/Engine/Source/Programs/Enterprise/Datasmith/DatasmithSketchUpRubyExporter/DatasmithSketchUpRuby2020.Target.cs b/Engine/Source/Programs/Enterprise/Datasmith/DatasmithSketchUpRubyExporter/DatasmithSketchUpRuby2020.Target.cs
--- a/Engine/Source/Programs/Enterprise/Datasmith/DatasmithSketchUpRubyExporter/DatasmithSketchUpRuby2020.Target.cs
+++ b/Engine/Source/Programs/Enterprise/Datasmith/DatasmithSketchUpRubyExporter/DatasmithSketchUpRuby2020.Target.cs
@@ -19,7 +19,6 @@ public abstract class DatasmithSketchUpRubyBaseTarget : TargetRules
 		WindowsPlatform.ModuleDefinitionFile = "Programs/Enterprise/Datasmith/DatasmithSketchUpRubyExporter/DatasmithSketchUpRubyExporter.def";
 
 		bBuildDeveloperTools = false;
-		bUseMallocProfiler = false;
 		bBuildWithEditorOnlyData = true;
 		bCompileAgainstEngine = false;
 		bCompileAgainstCoreUObject = true;
diff --git a/Engine/Source/Programs/Enterprise/Datasmith/DatasmithSolidworksExporter/DatasmithSolidworks2021.Target.cs b/Engine/Source/Programs/Enterprise/Datasmith/DatasmithSolidworksExporter/DatasmithSolidworks2021.Target.cs
--- a/Engine/Source/Programs/Enterprise/Datasmith/DatasmithSolidworksExporter/DatasmithSolidworks2021.Target.cs
+++ b/Engine/Source/Programs/Enterprise/Datasmith/DatasmithSolidworksExporter/DatasmithSolidworks2021.Target.cs
@@ -54,7 +54,6 @@ public abstract class DatasmithSolidworksBaseTarget : TargetRules
 		bShouldCompileAsDLL = true;
 		LinkType = TargetLinkType.Monolithic;
 		bBuildDeveloperTools = false;
-		bUseMallocProfiler = false;
 		bBuildWithEditorOnlyData = true;
 		bCompileAgainstEngine = false;
 		bCompileAgainstCoreUObject = false;
diff --git a/Engine/Source/Programs/EpicWebHelper/EpicWebHelper.Target.cs b/Engine/Source/Programs/EpicWebHelper/EpicWebHelper.Target.cs
--- a/Engine/Source/Programs/EpicWebHelper/EpicWebHelper.Target.cs
+++ b/Engine/Source/Programs/EpicWebHelper/EpicWebHelper.Target.cs
@@ -39,9 +39,6 @@ public class EpicWebHelperTarget : TargetRules
 		bCompileAgainstCoreUObject = false;
 		bBuildWithEditorOnlyData = true;
 
-		// Never use malloc profiling in CEFSubProcess.
-		bUseMallocProfiler = false;
-
 		// Force all shader formats to be built and included.
 		//bForceBuildShaderFormats = true;
 
diff --git a/Engine/Source/Programs/InterchangeWorker/InterchangeWorker.Target.cs b/Engine/Source/Programs/InterchangeWorker/InterchangeWorker.Target.cs
--- a/Engine/Source/Programs/InterchangeWorker/InterchangeWorker.Target.cs
+++ b/Engine/Source/Programs/InterchangeWorker/InterchangeWorker.Target.cs
@@ -18,10 +18,6 @@ public class InterchangeWorkerTarget : TargetRules
 		// Lean and mean
 		bBuildDeveloperTools = false;
 
-		// Never use malloc profiling in Unreal Header Tool.  We set this because often UHT is compiled right before the engine
-		// automatically by Unreal Build Tool, but if bUseMallocProfiler is defined, UHT can operate incorrectly.
-		bUseMallocProfiler = false;
-
 		// Editor-only data, however, is needed
 		bBuildWithEditorOnlyData = false;
 
diff --git a/Engine/Source/Programs/LiveCodingConsole/LiveCodingConsole.Target.cs b/Engine/Source/Programs/LiveCodingConsole/LiveCodingConsole.Target.cs
--- a/Engine/Source/Programs/LiveCodingConsole/LiveCodingConsole.Target.cs
+++ b/Engine/Source/Programs/LiveCodingConsole/LiveCodingConsole.Target.cs
@@ -14,7 +14,6 @@ public class LiveCodingConsoleTarget : TargetRules
 		LaunchModuleName = "LiveCodingConsole";
 
 		bBuildDeveloperTools = false;
-		bUseMallocProfiler = false;
 		bCompileWithPluginSupport = true;
 		bIncludePluginsForTargetPlatforms = true;
 		bWithLiveCoding = true;
diff --git a/Engine/Source/Programs/Mac/UnrealEditorServices/UnrealEditorServices.Target.cs b/Engine/Source/Programs/Mac/UnrealEditorServices/UnrealEditorServices.Target.cs
--- a/Engine/Source/Programs/Mac/UnrealEditorServices/UnrealEditorServices.Target.cs
+++ b/Engine/Source/Programs/Mac/UnrealEditorServices/UnrealEditorServices.Target.cs
@@ -13,7 +13,6 @@ public class UnrealEditorServicesTarget : TargetRules
 		LaunchModuleName = "UnrealEditorServices";
 
 		bBuildDeveloperTools = false;
-		bUseMallocProfiler = false;
 
 		bCompileICU = false;
 		// Editor-only data, however, is needed
diff --git a/Engine/Source/Programs/ReplicationSystemTest/ReplicationSystemTest.Target.cs b/Engine/Source/Programs/ReplicationSystemTest/ReplicationSystemTest.Target.cs
--- a/Engine/Source/Programs/ReplicationSystemTest/ReplicationSystemTest.Target.cs
+++ b/Engine/Source/Programs/ReplicationSystemTest/ReplicationSystemTest.Target.cs
@@ -12,8 +12,6 @@ public class ReplicationSystemTestTarget : TargetRules
 		LinkType = TargetLinkType.Monolithic;
 		LaunchModuleName = "ReplicationSystemTest";
 
-		bUseMallocProfiler = false;
-
 		// No editor-only data is needed
 		bBuildWithEditorOnlyData = false;
 
diff --git a/Engine/Source/Programs/ShaderCompileWorker/ShaderCompileWorker.Target.cs b/Engine/Source/Programs/ShaderCompileWorker/ShaderCompileWorker.Target.cs
--- a/Engine/Source/Programs/ShaderCompileWorker/ShaderCompileWorker.Target.cs
+++ b/Engine/Source/Programs/ShaderCompileWorker/ShaderCompileWorker.Target.cs
@@ -42,9 +42,6 @@ public class ShaderCompileWorkerTarget : TargetRules
 		bBuildWithEditorOnlyData = true;
 		bCompileCEF3 = false;
 
-		// Never use malloc profiling in ShaderCompileWorker.
-		bUseMallocProfiler = false;
-
 		// Force all shader formats to be built and included.
 		bForceBuildShaderFormats = true;
 
diff --git a/Engine/Source/Programs/Shared/EpicGames.Perforce.Native/DotNetPerforceLib.Target.cs b/Engine/Source/Programs/Shared/EpicGames.Perforce.Native/DotNetPerforceLib.Target.cs
--- a/Engine/Source/Programs/Shared/EpicGames.Perforce.Native/DotNetPerforceLib.Target.cs
+++ b/Engine/Source/Programs/Shared/EpicGames.Perforce.Native/DotNetPerforceLib.Target.cs
@@ -13,7 +13,6 @@ public class DotNetPerforceLibTarget : TargetRules
 		bShouldCompileAsDLL = true;
 
 		bBuildDeveloperTools = false;
-		bUseMallocProfiler = false;
 		bCompileICU = false;
 		bBuildWithEditorOnlyData = false;
 		bCompileAgainstEngine = false;
diff --git a/Engine/Source/Programs/UninstallHelper/UninstallHelper.Target.cs b/Engine/Source/Programs/UninstallHelper/UninstallHelper.Target.cs
--- a/Engine/Source/Programs/UninstallHelper/UninstallHelper.Target.cs
+++ b/Engine/Source/Programs/UninstallHelper/UninstallHelper.Target.cs
@@ -17,10 +17,6 @@ public class UninstallHelperTarget : TargetRules
 		// Lean and mean
 		bBuildDeveloperTools = false;
 
-		// Never use malloc profiling in Unreal Header Tool.  We set this because often UHT is compiled right before the engine
-		// automatically by Unreal Build Tool, but if bUseMallocProfiler is defined, UHT can operate incorrectly.
-		bUseMallocProfiler = false;
-		
 		// Currently this app is not linking against the engine, so we'll compile out references from Core to the rest of the engine
 		bCompileAgainstEngine = false;
 		bCompileAgainstCoreUObject = false;
diff --git a/Engine/Source/Programs/UnrealHeaderTool/UnrealHeaderTool.Target.cs b/Engine/Source/Programs/UnrealHeaderTool/UnrealHeaderTool.Target.cs
--- a/Engine/Source/Programs/UnrealHeaderTool/UnrealHeaderTool.Target.cs
+++ b/Engine/Source/Programs/UnrealHeaderTool/UnrealHeaderTool.Target.cs
@@ -12,10 +12,6 @@ public class UnrealHeaderToolTarget : TargetRules
 		IncludeOrderVersion = EngineIncludeOrderVersion.Latest;
 		LaunchModuleName = "UnrealHeaderTool";
 
-		// Never use malloc profiling in Unreal Header Tool.  We set this because often UHT is compiled right before the engine
-		// automatically by Unreal Build Tool, but if bUseMallocProfiler is defined, UHT can operate incorrectly.
-		bUseMallocProfiler = false;
-
         bCompileICU = false;
         // Editor-only data, however, is needed
 		bBuildWithEditorOnlyData = true;
diff --git a/Engine/Source/Programs/UnrealLightmass/UnrealLightmass.Target.cs b/Engine/Source/Programs/UnrealLightmass/UnrealLightmass.Target.cs
--- a/Engine/Source/Programs/UnrealLightmass/UnrealLightmass.Target.cs
+++ b/Engine/Source/Programs/UnrealLightmass/UnrealLightmass.Target.cs
@@ -17,10 +17,6 @@ public class UnrealLightmassTarget : TargetRules
 		// Lean and mean
 		bBuildDeveloperTools = false;
 
-		// Never use malloc profiling in Unreal Header Tool.  We set this because often UHT is compiled right before the engine
-		// automatically by Unreal Build Tool, but if bUseMallocProfiler is defined, UHT can operate incorrectly.
-		bUseMallocProfiler = false;
-
 		// Editor-only data, however, is needed
 		bBuildWithEditorOnlyData = true;
 
diff --git a/Engine/Source/Programs/UnrealObjectPtrTool/UnrealObjectPtrTool.Target.cs b/Engine/Source/Programs/UnrealObjectPtrTool/UnrealObjectPtrTool.Target.cs
--- a/Engine/Source/Programs/UnrealObjectPtrTool/UnrealObjectPtrTool.Target.cs
+++ b/Engine/Source/Programs/UnrealObjectPtrTool/UnrealObjectPtrTool.Target.cs
@@ -14,7 +14,6 @@ public class UnrealObjectPtrToolTarget : TargetRules
 		LaunchModuleName = "UnrealObjectPtrTool";
 
 		bBuildDeveloperTools			= false;
-		bUseMallocProfiler				= false;
 		bBuildWithEditorOnlyData		= false;
 		bCompileAgainstEngine			= false;
 		bCompileAgainstCoreUObject		= false;
diff --git a/Engine/Source/Programs/UnrealPak/UnrealPak.Target.cs b/Engine/Source/Programs/UnrealPak/UnrealPak.Target.cs
--- a/Engine/Source/Programs/UnrealPak/UnrealPak.Target.cs
+++ b/Engine/Source/Programs/UnrealPak/UnrealPak.Target.cs
@@ -26,7 +26,6 @@ public class UnrealPakTarget : TargetRules
 		bCompileICU 						= false;
 		bWithLiveCoding						= false;
 		bBuildDeveloperTools				= false;
-		bUseMallocProfiler					= false;
 		bBuildWithEditorOnlyData			= true;
 		bCompileAgainstEngine				= false;
 		bCompileAgainstCoreUObject			= true;
diff --git a/Engine/Source/Programs/UnrealVersionSelector/UnrealVersionSelector.Target.cs b/Engine/Source/Programs/UnrealVersionSelector/UnrealVersionSelector.Target.cs
--- a/Engine/Source/Programs/UnrealVersionSelector/UnrealVersionSelector.Target.cs
+++ b/Engine/Source/Programs/UnrealVersionSelector/UnrealVersionSelector.Target.cs
@@ -13,7 +13,6 @@ public class UnrealVersionSelectorTarget : TargetRules
 		LaunchModuleName = "UnrealVersionSelector";
 		
 		bBuildDeveloperTools = false;
-		bUseMallocProfiler = false;
 
 		bool bUsingSlate = (Target.Platform == UnrealTargetPlatform.Linux);
 
diff --git a/Engine/Source/Programs/UnrealVirtualizationTool/UnrealVirtualizationTool.Target.cs b/Engine/Source/Programs/UnrealVirtualizationTool/UnrealVirtualizationTool.Target.cs
--- a/Engine/Source/Programs/UnrealVirtualizationTool/UnrealVirtualizationTool.Target.cs
+++ b/Engine/Source/Programs/UnrealVirtualizationTool/UnrealVirtualizationTool.Target.cs
@@ -14,7 +14,6 @@ public class UnrealVirtualizationToolTarget : TargetRules
 		LaunchModuleName = "UnrealVirtualizationTool";
 
 		bBuildDeveloperTools = false;
-		bUseMallocProfiler = false;
 
 		bCompileAgainstEditor = false;
 		bBuildWithEditorOnlyData = true;
diff --git a/Engine/Source/Programs/Unsync/Unsync.Target.cs b/Engine/Source/Programs/Unsync/Unsync.Target.cs
--- a/Engine/Source/Programs/Unsync/Unsync.Target.cs
+++ b/Engine/Source/Programs/Unsync/Unsync.Target.cs
@@ -12,7 +12,6 @@ public class UnsyncTarget : TargetRules
 		LinkType = TargetLinkType.Monolithic;
 		LaunchModuleName = "Unsync";		
 		bBuildDeveloperTools = false;
-		bUseMallocProfiler = false;
 		bCompileICU = false;
 		bBuildWithEditorOnlyData = false;
 		bCompileAgainstEngine = false;
diff --git a/Engine/Source/Programs/VirtualProduction/TextureShare/TextureShareSDK.Target.cs b/Engine/Source/Programs/VirtualProduction/TextureShare/TextureShareSDK.Target.cs
--- a/Engine/Source/Programs/VirtualProduction/TextureShare/TextureShareSDK.Target.cs
+++ b/Engine/Source/Programs/VirtualProduction/TextureShare/TextureShareSDK.Target.cs
@@ -73,9 +73,6 @@ public class TextureShareSDKTarget : TargetRules
 		bUseLoggingInShipping = true;
 		bUseChecksInShipping = true;
 
-		// Disable memory profiling in the build
-		bUseMallocProfiler = false;
-
 		// Whether to include ICU unicode/i18n support in Core
 		bCompileICU = false;
 
diff --git a/Engine/Source/Programs/Windows/BootstrapPackagedGame/BootstrapPackagedGame.Target.cs b/Engine/Source/Programs/Windows/BootstrapPackagedGame/BootstrapPackagedGame.Target.cs
--- a/Engine/Source/Programs/Windows/BootstrapPackagedGame/BootstrapPackagedGame.Target.cs
+++ b/Engine/Source/Programs/Windows/BootstrapPackagedGame/BootstrapPackagedGame.Target.cs
@@ -16,7 +16,6 @@ public class BootstrapPackagedGameTarget : TargetRules
 		bUseStaticCRT = true;
 
 		bUseSharedPCHs = false;
-		bUseMallocProfiler = false;
 
 		// Disable all parts of the editor.
 		bBuildDeveloperTools = false;
diff --git a/Engine/Source/Programs/ZenDashboard/ZenDashboard.Target.cs b/Engine/Source/Programs/ZenDashboard/ZenDashboard.Target.cs
--- a/Engine/Source/Programs/ZenDashboard/ZenDashboard.Target.cs
+++ b/Engine/Source/Programs/ZenDashboard/ZenDashboard.Target.cs
@@ -24,7 +24,6 @@ public class ZenDashboardTarget : TargetRules
 		bCompileICU 					= false;
 		bWithLiveCoding					= false;
 		bBuildDeveloperTools			= false;
-		bUseMallocProfiler				= false;
 		bBuildWithEditorOnlyData		= false;
 		bCompileAgainstEngine			= false;
 		bCompileAgainstCoreUObject		= true;
