[BlankProgram]
- Add race detection support to BlankProgram

diff --git a/Engine/Source/Programs/BlankProgram/Private/BlankProgram.cpp b/Engine/Source/Programs/BlankProgram/Private/BlankProgram.cpp
--- a/Engine/Source/Programs/BlankProgram/Private/BlankProgram.cpp
+++ b/Engine/Source/Programs/BlankProgram/Private/BlankProgram.cpp
@@ -2,6 +2,7 @@
 
 #include "BlankProgram.h"
 
+#include "Sanitizer/RaceDetector.h"
 #include "RequiredProgramMainCPPInclude.h"
 
 DEFINE_LOG_CATEGORY_STATIC(LogBlankProgram, Log, All);
@@ -10,6 +11,10 @@ IMPLEMENT_APPLICATION(BlankProgram, "BlankProgram");
 
 INT32_MAIN_INT32_ARGC_TCHAR_ARGV()
 {
+#if USING_INSTRUMENTATION
+	UE::Sanitizer::RaceDetector::Initialize();
+#endif
+
 	FTaskTagScope Scope(ETaskTag::EGameThread);
 	ON_SCOPE_EXIT
 	{ 
@@ -18,6 +23,9 @@ INT32_MAIN_INT32_ARGC_TCHAR_ARGV()
 		FEngineLoop::AppPreExit();
 		FModuleManager::Get().UnloadModulesAtShutdown();
 		FEngineLoop::AppExit();
+#if USING_INSTRUMENTATION
+		UE::Sanitizer::RaceDetector::Shutdown();
+#endif
 	};
 
 	if (int32 Ret = GEngineLoop.PreInit(ArgC, ArgV))
