Set DefaultBuildSettings for the BlankProgram template. This prevents compilation errors when the template is used outside the engine folder.

No preflight as there is nothing that compiles the BlankProgram program
#Synced-CL 0

diff --git a/Engine/Source/Programs/BlankProgram/BlankProgram.Target.cs b/Engine/Source/Programs/BlankProgram/BlankProgram.Target.cs
--- a/Engine/Source/Programs/BlankProgram/BlankProgram.Target.cs
+++ b/Engine/Source/Programs/BlankProgram/BlankProgram.Target.cs
@@ -10,6 +10,7 @@ public class BlankProgramTarget : TargetRules
 	{
 		Type = TargetType.Program;
 		IncludeOrderVersion = EngineIncludeOrderVersion.Latest;
+		DefaultBuildSettings = BuildSettingsVersion.Latest;
 		LinkType = TargetLinkType.Monolithic;
 		LaunchModuleName = "BlankProgram";
 
